Added ascending and descending to collision analyzer for the ID column
[https://github.com/EpicGames/UnrealEngine/pull/506] (MyDigitalLife)

diff --git a/Engine/Source/Developer/CollisionAnalyzer/Private/CollisionAnalyzer.cpp b/Engine/Source/Developer/CollisionAnalyzer/Private/CollisionAnalyzer.cpp
--- a/Engine/Source/Developer/CollisionAnalyzer/Private/CollisionAnalyzer.cpp
+++ b/Engine/Source/Developer/CollisionAnalyzer/Private/CollisionAnalyzer.cpp
@@ -8,18 +8,18 @@
 
 
 void FCollisionAnalyzer::CaptureQuery(	const FVector& Start, 
-									  const FVector& End, 
-									  const FQuat& Rot, 
-									  ECAQueryType::Type QueryType, 
-									  ECAQueryShape::Type QueryShape, 
-									  const FVector& Dims, 
-									  ECollisionChannel TraceChannel, 
-									  const struct FCollisionQueryParams& Params, 
-									  const FCollisionResponseParams&	ResponseParams,
-									  const FCollisionObjectQueryParams&	ObjectParams,
-									  const TArray<FHitResult>& Results, 
-									  const TArray<FHitResult>& TouchAllResults,
-									  double CPUTime ) 
+										const FVector& End, 
+										const FQuat& Rot, 
+										ECAQueryType::Type QueryType, 
+										ECAQueryShape::Type QueryShape, 
+										const FVector& Dims, 
+										ECollisionChannel TraceChannel, 
+										const struct FCollisionQueryParams& Params, 
+										const FCollisionResponseParams&	ResponseParams,
+										const FCollisionObjectQueryParams&	ObjectParams,
+										const TArray<FHitResult>& Results, 
+										const TArray<FHitResult>& TouchAllResults,
+										double CPUTime) 
 {
 	if(bIsRecording)
 	{
@@ -38,7 +38,9 @@ void FCollisionAnalyzer::CaptureQuery(	const FVector& Start,
 		NewQuery.Results = Results;
 		NewQuery.TouchAllResults = TouchAllResults;
 		NewQuery.FrameNum = CurrentFrameNum;
-		NewQuery.CPUTime = CPUTime * 1000.f;
+		NewQuery.CPUTime = CPUTime * 1000.f;		
+		NewQuery.ID = NewQueryId;
+
 
 		QueryAddedEvent.Broadcast();
 	}
diff --git a/Engine/Source/Developer/CollisionAnalyzer/Private/CollisionAnalyzer.h b/Engine/Source/Developer/CollisionAnalyzer/Private/CollisionAnalyzer.h
--- a/Engine/Source/Developer/CollisionAnalyzer/Private/CollisionAnalyzer.h
+++ b/Engine/Source/Developer/CollisionAnalyzer/Private/CollisionAnalyzer.h
@@ -19,6 +19,7 @@ struct FCAQuery
 	TArray<FHitResult>			TouchAllResults;
 	int32						FrameNum;
 	float						CPUTime; /** In ms */
+	int32						ID;
 };
 
 /** Actual implementation of CollisionAnalyzer, private inside module */
diff --git a/Engine/Source/Developer/CollisionAnalyzer/Private/SCollisionAnalyzer.cpp b/Engine/Source/Developer/CollisionAnalyzer/Private/SCollisionAnalyzer.cpp
--- a/Engine/Source/Developer/CollisionAnalyzer/Private/SCollisionAnalyzer.cpp
+++ b/Engine/Source/Developer/CollisionAnalyzer/Private/SCollisionAnalyzer.cpp
@@ -6,6 +6,9 @@
 
 static const int32 NumDrawRecentQueries = 10;
 
+static FName TimeColumnName(TEXT("Name"));
+static FName IDColumnName(TEXT("ID"));
+
 BEGIN_SLATE_FUNCTION_BUILD_OPTIMIZATION
 void SCollisionAnalyzer::Construct(const FArguments& InArgs, FCollisionAnalyzer* InAnalyzer)
 {
@@ -15,6 +18,7 @@ void SCollisionAnalyzer::Construct(const FArguments& InArgs, FCollisionAnalyzer*
 	MinCPUFilterTime = -1.f;
 	GroupBy = EQueryGroupMode::Ungrouped;
 	SortBy = EQuerySortMode::ByID;
+	SortDirection = EColumnSortMode::Descending;
 	TotalNumQueries = 0;
 
 	ChildSlot
@@ -81,7 +85,7 @@ void SCollisionAnalyzer::Construct(const FArguments& InArgs, FCollisionAnalyzer*
 					.HeaderRow(
 						SNew(SHeaderRow)
 						// ID
-						+SHeaderRow::Column("ID")
+						+SHeaderRow::Column(IDColumnName)
 						.SortMode(this, &SCollisionAnalyzer::GetIDSortMode)
 						.OnSort(this, &SCollisionAnalyzer::OnSortByChanged)
 						.HAlignCell(HAlign_Left)
@@ -241,7 +245,7 @@ void SCollisionAnalyzer::Construct(const FArguments& InArgs, FCollisionAnalyzer*
 							.ToolTipText( LOCTEXT("NumberTouchTooltip", "Number of touching results") )
 						]
 						// CPU time
-						+SHeaderRow::Column("Time")
+						+SHeaderRow::Column(TimeColumnName)
 						.SortMode(this, &SCollisionAnalyzer::GetTimeSortMode)
 						.OnSort(this, &SCollisionAnalyzer::OnSortByChanged)
 						.FixedWidth(48)
@@ -480,6 +484,37 @@ struct FCompareQueryByCPUTime
 	}
 };
 
+/** Functor for comparing query by ID */
+struct FCompareQueryByID
+{
+	FCollisionAnalyzer* Analyzer;
+	EColumnSortMode::Type SortMode;
+	
+	FCompareQueryByID(FCollisionAnalyzer* InAnalyzer, EColumnSortMode::Type InSortMode)
+		 : Analyzer(InAnalyzer), SortMode(InSortMode)
+	 {}
+	
+	FORCEINLINE bool operator()(const TSharedPtr<FQueryTreeItem> A, const TSharedPtr<FQueryTreeItem> B) const
+	{
+		check(A.IsValid());
+		check(!A->bIsGroup);
+		check(B.IsValid());
+		check(!A->bIsGroup);
+		const FCAQuery& QueryA = Analyzer->Queries[A->QueryIndex];
+		const FCAQuery& QueryB = Analyzer->Queries[B->QueryIndex];
+		
+		if (SortMode == EColumnSortMode::Descending)
+		{
+			return (QueryA.ID < QueryB.ID);
+		}
+		else
+		{
+			return (QueryA.ID > QueryB.ID);
+		}
+	}
+};
+
+
 /** Functor for comparing group by CPU time */
 struct FCompareGroupByCPUTime
 {
@@ -507,6 +542,10 @@ void SCollisionAnalyzer::AddQueryToGroupedQueries(int32 NewQueryIndex, bool bPer
 		{
 			GroupedQueries.Sort( FCompareQueryByCPUTime(Analyzer) );
 		}
+		else if (SortBy == EQuerySortMode::ByID)
+		{
+			GroupedQueries.Sort( FCompareQueryByID(Analyzer, SortDirection) );
+		}
 	}
 	// If we are grouping..
 	else 
@@ -596,6 +635,10 @@ void SCollisionAnalyzer::RebuildFilteredList()
 			}
 		}
 	}
+	else if (SortBy == EQuerySortMode::ByID)
+	{
+		GroupedQueries.Sort(FCompareQueryByID(Analyzer, SortDirection));
+	}
 	
 	// When underlying array changes, refresh list
 	QueryTreeWidget->RequestTreeRefresh();
@@ -689,11 +732,25 @@ void SCollisionAnalyzer::FilterTextCommitted(const FText& CommentText, ETextComm
 
 void SCollisionAnalyzer::OnSortByChanged(const EColumnSortPriority::Type SortPriority, const FName& ColumnName, const EColumnSortMode::Type NewSortMode)
 {
-	SortBy = EQuerySortMode::ByID;
-
-	if(ColumnName == TEXT("Time"))
+	if(ColumnName == TimeColumnName)
 	{
+		// Only makes sense to sort by time decreasing
 		SortBy = EQuerySortMode::ByTime;
+		SortDirection = EColumnSortMode::Descending;
+	}
+	else if(ColumnName == IDColumnName)
+	{
+		// If already sorting by ID, flip direction
+		if (SortBy == EQuerySortMode::ByID)
+		{
+			SortDirection = (SortDirection == EColumnSortMode::Descending) ? EColumnSortMode::Ascending : EColumnSortMode::Descending;
+		}
+		// If not, sort by ID, and default to ascending
+		else
+		{
+			SortBy = EQuerySortMode::ByID;
+			SortDirection = EColumnSortMode::Descending;
+		}
 	}
 
 	RebuildFilteredList();
@@ -701,12 +758,12 @@ void SCollisionAnalyzer::OnSortByChanged(const EColumnSortPriority::Type SortPri
 
 EColumnSortMode::Type SCollisionAnalyzer::GetIDSortMode() const
 {
-	return (SortBy == EQuerySortMode::ByID) ? EColumnSortMode::Ascending : EColumnSortMode::None;
+	return (SortBy == EQuerySortMode::ByID) ? SortDirection : EColumnSortMode::None;
 }
 
 EColumnSortMode::Type SCollisionAnalyzer::GetTimeSortMode() const
 {
-	return (SortBy == EQuerySortMode::ByTime) ? EColumnSortMode::Descending : EColumnSortMode::None;
+	return (SortBy == EQuerySortMode::ByTime) ? SortDirection : EColumnSortMode::None;
 }
 
 FString SCollisionAnalyzer::QueryShapeToString(ECAQueryShape::Type QueryShape)
diff --git a/Engine/Source/Developer/CollisionAnalyzer/Private/SCollisionAnalyzer.h b/Engine/Source/Developer/CollisionAnalyzer/Private/SCollisionAnalyzer.h
--- a/Engine/Source/Developer/CollisionAnalyzer/Private/SCollisionAnalyzer.h
+++ b/Engine/Source/Developer/CollisionAnalyzer/Private/SCollisionAnalyzer.h
@@ -127,7 +127,8 @@ public:
 	EQueryGroupMode::Type			GroupBy;
 	/** Current way we are sorting queries */
 	EQuerySortMode::Type			SortBy;
-
+	/** Current way we are setting ID sort direction */
+	EColumnSortMode::Type			SortDirection;
 private:
 	/** Called when the queries array in the CollisionAnalyzer change */
 	void OnQueriesChanged();
