Added custom StyleSet to Collision Analyzer to remove dependency on EditorStyleSet (so that it can be used in non-editor developer builds).

diff --git a/Engine/Source/Developer/CollisionAnalyzer/Private/CollisionAnalyzerModule.cpp b/Engine/Source/Developer/CollisionAnalyzer/Private/CollisionAnalyzerModule.cpp
--- a/Engine/Source/Developer/CollisionAnalyzer/Private/CollisionAnalyzerModule.cpp
+++ b/Engine/Source/Developer/CollisionAnalyzer/Private/CollisionAnalyzerModule.cpp
@@ -14,6 +14,8 @@ DEFINE_LOG_CATEGORY(LogCollisionAnalyzer);
 
 void FCollisionAnalyzerModule::StartupModule() 
 {
+	FCollisionAnalyzerStyle::Initialize();
+
 	CollisionAnalyzer = new FCollisionAnalyzer();
 
 	FGlobalTabmanager::Get()->RegisterNomadTabSpawner(CollisionAnalyzerModule::CollisionAnalyzerApp, FOnSpawnTab::CreateRaw(this, &FCollisionAnalyzerModule::SpawnCollisionAnalyzerTab))
@@ -30,6 +32,8 @@ void FCollisionAnalyzerModule::ShutdownModule()
 		delete CollisionAnalyzer;
 		CollisionAnalyzer = NULL;
 	}
+
+	FCollisionAnalyzerStyle::Shutdown();
 }
 
 TSharedRef<SDockTab> FCollisionAnalyzerModule::SpawnCollisionAnalyzerTab(const FSpawnTabArgs& Args)
diff --git a/Engine/Source/Developer/CollisionAnalyzer/Private/CollisionAnalyzerPCH.h b/Engine/Source/Developer/CollisionAnalyzer/Private/CollisionAnalyzerPCH.h
--- a/Engine/Source/Developer/CollisionAnalyzer/Private/CollisionAnalyzerPCH.h
+++ b/Engine/Source/Developer/CollisionAnalyzer/Private/CollisionAnalyzerPCH.h
@@ -10,6 +10,7 @@
 #include "ICollisionAnalyzer.h"
 
 #include "CollisionAnalyzer.h"
+#include "CollisionAnalyzerStyle.h"
 
 #include "SCAQueryTableRow.h"
 #include "SCAQueryDetails.h"
diff --git a/Engine/Source/Developer/CollisionAnalyzer/Private/CollisionAnalyzerStyle.cpp b/Engine/Source/Developer/CollisionAnalyzer/Private/CollisionAnalyzerStyle.cpp
new file mode 100644
--- /dev/null
+++ b/Engine/Source/Developer/CollisionAnalyzer/Private/CollisionAnalyzerStyle.cpp
@@ -0,0 +1,76 @@
+// Copyright 1998-2015 Epic Games, Inc. All Rights Reserved.
+
+#include "CollisionAnalyzerPCH.h"
+#include "SlateBasics.h"
+#include "CollisionAnalyzerStyle.h"
+
+#define IMAGE_BRUSH( RelativePath, ... ) FSlateImageBrush( RootToContentDir( RelativePath, TEXT(".png") ), __VA_ARGS__ )
+#define BOX_BRUSH( RelativePath, ... ) FSlateBoxBrush( RootToContentDir( RelativePath, TEXT(".png") ), __VA_ARGS__ )
+#define BORDER_BRUSH( RelativePath, ... ) FSlateBorderBrush( RootToContentDir( RelativePath, TEXT(".png") ), __VA_ARGS__ )
+#define TTF_CORE_FONT( RelativePath, ... ) FSlateFontInfo( RootToCoreContentDir( RelativePath, TEXT(".ttf") ), __VA_ARGS__ )
+
+TSharedPtr< FCollisionAnalyzerStyle::FStyle > FCollisionAnalyzerStyle::StyleInstance;
+
+FCollisionAnalyzerStyle::FStyle::FStyle()
+	: FSlateStyleSet("CollisionAnalyzerStyle")
+{
+
+}
+
+void FCollisionAnalyzerStyle::FStyle::Initialize()
+{
+
+
+	StyleInstance->SetContentRoot(FPaths::EngineContentDir() / TEXT("Editor/Slate"));
+	StyleInstance->SetCoreContentRoot(FPaths::EngineContentDir() / TEXT("Slate"));
+
+	const FLinearColor SelectionColor(0.728f, 0.364f, 0.003f);
+	const FLinearColor SelectionColor_Pressed(0.701f, 0.225f, 0.003f);
+
+	const FCheckBoxStyle ToggleButtonStyle = FCheckBoxStyle()
+		.SetCheckBoxType(ESlateCheckBoxType::ToggleButton)
+		.SetUncheckedImage(BOX_BRUSH("Common/RoundedSelection_16x", 4.0f / 16.0f, FLinearColor(1, 1, 1, 0.1f)))
+		.SetUncheckedHoveredImage(BOX_BRUSH("Common/RoundedSelection_16x", 4.0f / 16.0f, SelectionColor))
+		.SetUncheckedPressedImage(BOX_BRUSH("Common/RoundedSelection_16x", 4.0f / 16.0f, SelectionColor_Pressed))
+		.SetCheckedImage(BOX_BRUSH("Common/RoundedSelection_16x", 4.0f / 16.0f, SelectionColor_Pressed))
+		.SetCheckedHoveredImage(BOX_BRUSH("Common/RoundedSelection_16x", 4.0f / 16.0f, SelectionColor))
+		.SetCheckedPressedImage(BOX_BRUSH("Common/RoundedSelection_16x", 4.0f / 16.0f, SelectionColor_Pressed));
+	StyleInstance->Set("ToggleButtonCheckbox", ToggleButtonStyle);
+
+	const FVector2D Icon16x16(16.0f, 16.0f);
+	const FVector2D Icon24x24(24.0f, 24.0f);
+
+	StyleInstance->Set("ToolBar.Background", new BOX_BRUSH("Common/GroupBorder", FMargin(4.0f / 16.0f)));
+	StyleInstance->Set("CollisionAnalyzer.TabIcon", new IMAGE_BRUSH("Icons/icon_tab_CollisionAnalyser_16x", Icon16x16));
+	StyleInstance->Set("CollisionAnalyzer.Record", new IMAGE_BRUSH("Icons/CA_Record", Icon24x24));
+	StyleInstance->Set("CollisionAnalyzer.Stop", new IMAGE_BRUSH("Icons/CA_Stop", Icon24x24));
+	StyleInstance->Set("CollisionAnalyzer.ShowRecent", new IMAGE_BRUSH("Icons/CA_ShowRecent", Icon24x24));
+	StyleInstance->Set("CollisionAnalyzer.Group", new IMAGE_BRUSH("Icons/CA_Group", FVector2D(10, 18)));
+	StyleInstance->Set("CollisionAnalyzer.GroupBackground", new BOX_BRUSH("Icons/CA_GroupBackground", FMargin(4.f / 16.f)));
+	StyleInstance->Set("Menu.Background", new BOX_BRUSH("Old/Menu_Background", FMargin(8.0f / 64.0f)));
+	StyleInstance->Set("BoldFont", TTF_CORE_FONT("Fonts/Roboto-Bold", 9));
+
+	FSlateStyleRegistry::RegisterSlateStyle(*StyleInstance.Get());
+}
+
+
+void FCollisionAnalyzerStyle::Initialize()
+{
+	StyleInstance = MakeShareable(new FCollisionAnalyzerStyle::FStyle);
+	StyleInstance->Initialize();
+}
+
+void FCollisionAnalyzerStyle::Shutdown()
+{
+	if (StyleInstance.IsValid())
+	{
+		FSlateStyleRegistry::UnRegisterSlateStyle(*StyleInstance.Get());
+		ensure(StyleInstance.IsUnique());
+		StyleInstance.Reset();
+	}
+}
+
+TSharedPtr< ISlateStyle > FCollisionAnalyzerStyle::Get()
+{
+	return StyleInstance;
+}
diff --git a/Engine/Source/Developer/CollisionAnalyzer/Private/CollisionAnalyzerStyle.h b/Engine/Source/Developer/CollisionAnalyzer/Private/CollisionAnalyzerStyle.h
new file mode 100644
--- /dev/null
+++ b/Engine/Source/Developer/CollisionAnalyzer/Private/CollisionAnalyzerStyle.h
@@ -0,0 +1,30 @@
+// Copyright 1998-2015 Epic Games, Inc. All Rights Reserved.
+
+#pragma once
+
+
+#include "EditorStyle.h"
+
+/**
+* Declares the Collision Analyzer window's visual style.
+*/
+class FCollisionAnalyzerStyle
+{
+public:
+
+	class FStyle : public FSlateStyleSet
+	{
+	public:
+		FStyle();
+
+		void Initialize();
+	};
+
+	static void Initialize();
+
+	static void Shutdown();
+
+	static TSharedPtr<class ISlateStyle> Get();
+
+	static TSharedPtr< FCollisionAnalyzerStyle::FStyle > StyleInstance;
+};
diff --git a/Engine/Source/Developer/CollisionAnalyzer/Private/SCAQueryDetails.cpp b/Engine/Source/Developer/CollisionAnalyzer/Private/SCAQueryDetails.cpp
--- a/Engine/Source/Developer/CollisionAnalyzer/Private/SCAQueryDetails.cpp
+++ b/Engine/Source/Developer/CollisionAnalyzer/Private/SCAQueryDetails.cpp
@@ -153,7 +153,7 @@ void SCAQueryDetails::Construct(const FArguments& InArgs, TSharedPtr<SCollisionA
 		.AutoHeight()
 		[
 			SNew( SBorder )
-			.BorderImage( FEditorStyle::GetBrush( "ToolBar.Background" ) )
+			.BorderImage(FCollisionAnalyzerStyle::Get()->GetBrush("ToolBar.Background"))
 			[
 				SNew(SHorizontalBox)
 				// Left is start/end locations
@@ -207,7 +207,7 @@ void SCAQueryDetails::Construct(const FArguments& InArgs, TSharedPtr<SCollisionA
 		.FillHeight(1) 
 		[
 			SNew(SBorder)
-			.BorderImage(FEditorStyle::GetBrush("Menu.Background"))
+			.BorderImage(FCollisionAnalyzerStyle::Get()->GetBrush("Menu.Background"))
 			.Padding(1.0)
 			[
 				SAssignNew(ResultListWidget, SListView< TSharedPtr<FCAHitInfo> >)
diff --git a/Engine/Source/Developer/CollisionAnalyzer/Private/SCAQueryTableRow.cpp b/Engine/Source/Developer/CollisionAnalyzer/Private/SCAQueryTableRow.cpp
--- a/Engine/Source/Developer/CollisionAnalyzer/Private/SCAQueryTableRow.cpp
+++ b/Engine/Source/Developer/CollisionAnalyzer/Private/SCAQueryTableRow.cpp
@@ -12,7 +12,7 @@ void SCAQueryTableRow::Construct(const FArguments& InArgs, const TSharedRef<STab
 
 	if(Item->bIsGroup)
 	{
-		BorderImage = FEditorStyle::GetBrush("CollisionAnalyzer.GroupBackground");
+		BorderImage = FCollisionAnalyzerStyle::Get()->GetBrush("CollisionAnalyzer.GroupBackground");
 	}
 }
 
@@ -30,25 +30,25 @@ TSharedRef<SWidget> SCAQueryTableRow::GenerateWidgetForColumn(const FName& Colum
 		else if (ColumnName == TEXT("Frame") && OwnerAnalyzerWidget->GroupBy == EQueryGroupMode::ByFrameNum)
 		{
 			return	SNew(STextBlock)
-					.Font(FEditorStyle::GetFontStyle("BoldFont"))
+					.Font(FCollisionAnalyzerStyle::Get()->GetFontStyle("BoldFont"))
 					.Text(FText::AsNumber(Item->FrameNum));
 		}
 		else if (ColumnName == TEXT("Tag") && OwnerAnalyzerWidget->GroupBy == EQueryGroupMode::ByTag)
 		{
 			return	SNew(STextBlock)
-					.Font(FEditorStyle::GetFontStyle("BoldFont"))
+					.Font(FCollisionAnalyzerStyle::Get()->GetFontStyle("BoldFont"))
 					.Text(FText::FromName(Item->GroupName));
 		}
 		else if (ColumnName == TEXT("Owner") && OwnerAnalyzerWidget->GroupBy == EQueryGroupMode::ByOwnerTag)
 		{
 			return	SNew(STextBlock)
-					.Font(FEditorStyle::GetFontStyle("BoldFont"))
+					.Font(FCollisionAnalyzerStyle::Get()->GetFontStyle("BoldFont"))
 					.Text(FText::FromName(Item->GroupName));
 		}
 		else if (ColumnName == TEXT("Time"))
 		{
 			return	SNew(STextBlock)
-					.Font(FEditorStyle::GetFontStyle("BoldFont"))
+					.Font(FCollisionAnalyzerStyle::Get()->GetFontStyle("BoldFont"))
 					.Text(this, &SCAQueryTableRow::GetTotalTimeText);
 		}
 	}
diff --git a/Engine/Source/Developer/CollisionAnalyzer/Private/SCollisionAnalyzer.cpp b/Engine/Source/Developer/CollisionAnalyzer/Private/SCollisionAnalyzer.cpp
--- a/Engine/Source/Developer/CollisionAnalyzer/Private/SCollisionAnalyzer.cpp
+++ b/Engine/Source/Developer/CollisionAnalyzer/Private/SCollisionAnalyzer.cpp
@@ -29,7 +29,7 @@ void SCollisionAnalyzer::Construct(const FArguments& InArgs, FCollisionAnalyzer*
 		.AutoHeight()
 		[
 			SNew( SBorder )
-			.BorderImage( FEditorStyle::GetBrush( "ToolBar.Background" ) )
+			.BorderImage( FCollisionAnalyzerStyle::Get()->GetBrush( "ToolBar.Background" ) )
 			[
 				SNew(SHorizontalBox)
 				// Record button
@@ -51,13 +51,13 @@ void SCollisionAnalyzer::Construct(const FArguments& InArgs, FCollisionAnalyzer*
 				.Padding(1)
 				[
 					SNew(SCheckBox)
-					.Style(FEditorStyle::Get(), "ToggleButtonCheckbox")
+					.Style(FCollisionAnalyzerStyle::Get(), "ToggleButtonCheckbox")
 					.OnCheckStateChanged(this, &SCollisionAnalyzer::OnDrawRecentChanged)
 					.IsChecked(this, &SCollisionAnalyzer::GetDrawRecentState)
 					.Content()
 					[
 						SNew(SImage)
-						.Image(FEditorStyle::GetBrush("CollisionAnalyzer.ShowRecent"))
+						.Image(FCollisionAnalyzerStyle::Get()->GetBrush("CollisionAnalyzer.ShowRecent"))
 					]
 				]
 			]
@@ -72,7 +72,7 @@ void SCollisionAnalyzer::Construct(const FArguments& InArgs, FCollisionAnalyzer*
 			.Value(2)
 			[
 				SNew(SBorder)
-				.BorderImage(FEditorStyle::GetBrush("Menu.Background"))
+				.BorderImage(FCollisionAnalyzerStyle::Get()->GetBrush("Menu.Background"))
 				.Padding(1.0)
 				[
 					SAssignNew(QueryTreeWidget, STreeView< TSharedPtr<FQueryTreeItem> >)
@@ -123,13 +123,13 @@ void SCollisionAnalyzer::Construct(const FArguments& InArgs, FCollisionAnalyzer*
 								.AutoWidth()
 								[
 									SNew(SCheckBox)
-									.Style(FEditorStyle::Get(), "ToggleButtonCheckbox")
+									.Style(FCollisionAnalyzerStyle::Get(), "ToggleButtonCheckbox")
 									.OnCheckStateChanged(this, &SCollisionAnalyzer::OnGroupByFrameChanged)
 									.IsChecked( this, &SCollisionAnalyzer::GetGroupByFrameState)
 									.Content()
 									[
 										SNew(SImage)
-										.Image(FEditorStyle::GetBrush("CollisionAnalyzer.Group"))
+										.Image(FCollisionAnalyzerStyle::Get()->GetBrush("CollisionAnalyzer.Group"))
 									]
 								]
 							]
@@ -177,13 +177,13 @@ void SCollisionAnalyzer::Construct(const FArguments& InArgs, FCollisionAnalyzer*
 								.AutoWidth()
 								[
 									SNew(SCheckBox)
-									.Style(FEditorStyle::Get(), "ToggleButtonCheckbox")
+									.Style(FCollisionAnalyzerStyle::Get(), "ToggleButtonCheckbox")
 									.OnCheckStateChanged(this, &SCollisionAnalyzer::OnGroupByTagChanged)
 									.IsChecked( this, &SCollisionAnalyzer::GetGroupByTagState)
 									.Content()
 									[
 										SNew(SImage)
-										.Image(FEditorStyle::GetBrush("CollisionAnalyzer.Group"))
+										.Image(FCollisionAnalyzerStyle::Get()->GetBrush("CollisionAnalyzer.Group"))
 									]
 								]
 							]
@@ -217,13 +217,13 @@ void SCollisionAnalyzer::Construct(const FArguments& InArgs, FCollisionAnalyzer*
 								.AutoWidth()
 								[
 									SNew(SCheckBox)
-									.Style(FEditorStyle::Get(), "ToggleButtonCheckbox")
+									.Style(FCollisionAnalyzerStyle::Get(), "ToggleButtonCheckbox")
 									.OnCheckStateChanged(this, &SCollisionAnalyzer::OnGroupByOwnerChanged)
 									.IsChecked( this, &SCollisionAnalyzer::GetGroupByOwnerState)
 									.Content()
 									[
 										SNew(SImage)
-										.Image(FEditorStyle::GetBrush("CollisionAnalyzer.Group"))
+										.Image(FCollisionAnalyzerStyle::Get()->GetBrush("CollisionAnalyzer.Group"))
 									]
 								]
 							]
@@ -281,7 +281,7 @@ void SCollisionAnalyzer::Construct(const FArguments& InArgs, FCollisionAnalyzer*
 		.AutoHeight()
 		[
 			SNew( SBorder )
-			.BorderImage( FEditorStyle::GetBrush( "ToolBar.Background" ) )
+			.BorderImage(FCollisionAnalyzerStyle::Get()->GetBrush("ToolBar.Background"))
 			[
 				SNew(STextBlock)
 				.Text(this, &SCollisionAnalyzer::GetStatusText)
@@ -651,12 +651,12 @@ const FSlateBrush* SCollisionAnalyzer::GetRecordButtonBrush() const
 	if(Analyzer->IsRecording())
 	{
 		// If recording, show stop button
-		return 	FEditorStyle::GetBrush("CollisionAnalyzer.Stop");
+		return FCollisionAnalyzerStyle::Get()->GetBrush("CollisionAnalyzer.Stop");
 	}
 	else
 	{
 		// If stopped, show record button
-		return 	FEditorStyle::GetBrush("CollisionAnalyzer.Record");
+		return FCollisionAnalyzerStyle::Get()->GetBrush("CollisionAnalyzer.Record");
 	}
 }
 
diff --git a/Engine/Source/Editor/EditorStyle/Private/SlateEditorStyle.cpp b/Engine/Source/Editor/EditorStyle/Private/SlateEditorStyle.cpp
--- a/Engine/Source/Editor/EditorStyle/Private/SlateEditorStyle.cpp
+++ b/Engine/Source/Editor/EditorStyle/Private/SlateEditorStyle.cpp
@@ -1709,15 +1709,6 @@ void FSlateEditorStyle::FStyle::SetupGeneralStyles()
 		Set( "EditorLiveStreaming.BroadcastButton", new IMAGE_BRUSH( "Animation/Record", Icon20x20 ) );
 	}
 
-	// Collision Analyzer
-	{
-		Set( "CollisionAnalyzer.TabIcon", new IMAGE_BRUSH( "Icons/icon_tab_CollisionAnalyser_16x", Icon16x16 ) );
-		Set( "CollisionAnalyzer.Record", new IMAGE_BRUSH( "Icons/CA_Record", Icon24x24 ) );
-		Set( "CollisionAnalyzer.Stop", new IMAGE_BRUSH( "Icons/CA_Stop", Icon24x24 ) );
-		Set( "CollisionAnalyzer.ShowRecent", new IMAGE_BRUSH( "Icons/CA_ShowRecent", Icon24x24 ) );
-		Set( "CollisionAnalyzer.Group", new IMAGE_BRUSH( "Icons/CA_Group", FVector2D(10,18) ) );
-		Set( "CollisionAnalyzer.GroupBackground", new BOX_BRUSH( "/Icons/CA_GroupBackground", FMargin(4.f/16.f) ) );
-	}
 #endif // WITH_EDITOR || IS_PROGRAM
 
 #if WITH_EDITOR
