Replace usages of toolbar.background with toolpanel.group border

diff --git a/Engine/Plugins/Editor/GameplayTagsEditor/Source/GameplayTagsEditor/Private/GameplayTagReferenceHelperDetails.cpp b/Engine/Plugins/Editor/GameplayTagsEditor/Source/GameplayTagsEditor/Private/GameplayTagReferenceHelperDetails.cpp
--- a/Engine/Plugins/Editor/GameplayTagsEditor/Source/GameplayTagsEditor/Private/GameplayTagReferenceHelperDetails.cpp
+++ b/Engine/Plugins/Editor/GameplayTagsEditor/Source/GameplayTagsEditor/Private/GameplayTagReferenceHelperDetails.cpp
@@ -78,7 +78,7 @@ void FGameplayTagReferenceHelperDetails::CustomizeHeader( TSharedRef<IPropertyHa
 	.NameContent()
 	[
 		SNew(SBorder)
-		.BorderImage(FEditorStyle::GetBrush("ToolBar.Background"))
+		.BorderImage(FEditorStyle::GetBrush("ToolPanel.GroupBorder"))
 		[
 			SNew(SGameplayTagReferenceTree)
 			.ItemHeight(24)
diff --git a/Engine/Plugins/Runtime/GameplayAbilities/Source/GameplayAbilitiesEditor/Private/SGameplayCueEditor.cpp b/Engine/Plugins/Runtime/GameplayAbilities/Source/GameplayAbilitiesEditor/Private/SGameplayCueEditor.cpp
--- a/Engine/Plugins/Runtime/GameplayAbilities/Source/GameplayAbilitiesEditor/Private/SGameplayCueEditor.cpp
+++ b/Engine/Plugins/Runtime/GameplayAbilities/Source/GameplayAbilitiesEditor/Private/SGameplayCueEditor.cpp
@@ -1312,7 +1312,7 @@ void SGameplayCueEditorImpl::Construct(const FArguments& InArgs)
 		.FillHeight(1.0f)
 		[
 			SNew(SBorder)
-			.BorderImage(FEditorStyle::GetBrush("ToolBar.Background"))
+			.BorderImage(FEditorStyle::GetBrush("ToolPanel.GroupBorder"))
 			[
 				SAssignNew(GameplayCueTreeView, SGameplayCueTreeView)
 				.ItemHeight(24)
diff --git a/Engine/Source/Developer/CollisionAnalyzer/Private/CollisionAnalyzerStyle.cpp b/Engine/Source/Developer/CollisionAnalyzer/Private/CollisionAnalyzerStyle.cpp
--- a/Engine/Source/Developer/CollisionAnalyzer/Private/CollisionAnalyzerStyle.cpp
+++ b/Engine/Source/Developer/CollisionAnalyzer/Private/CollisionAnalyzerStyle.cpp
@@ -45,7 +45,6 @@ void FCollisionAnalyzerStyle::FStyle::Initialize()
 		.SetPressed(BOX_BRUSH("Common/RoundedSelection_16x", 4.0f / 16.0f, SelectionColor_Pressed));
 	StyleInstance->Set("CommonButton", CommonButtonStyle);
 
-	StyleInstance->Set("ToolBar.Background", new BOX_BRUSH("Common/GroupBorder", FMargin(4.0f / 16.0f)));
 	StyleInstance->Set("CollisionAnalyzer.TabIcon", new IMAGE_BRUSH("Icons/icon_tab_CollisionAnalyser_16x", Icon16x16));
 	StyleInstance->Set("CollisionAnalyzer.Record", new IMAGE_BRUSH("Icons/CA_Record", Icon24x24));
 	StyleInstance->Set("CollisionAnalyzer.Stop", new IMAGE_BRUSH("Icons/CA_Stop", Icon24x24));
@@ -54,10 +53,10 @@ void FCollisionAnalyzerStyle::FStyle::Initialize()
 	StyleInstance->Set("CollisionAnalyzer.GroupBackground", new BOX_BRUSH("Icons/CA_GroupBackground", FMargin(4.f / 16.f)));
 	StyleInstance->Set("CollisionAnalyzer.Save", new IMAGE_BRUSH("Icons/icon_file_save_40x", Icon24x24));
 	StyleInstance->Set("CollisionAnalyzer.Load", new IMAGE_BRUSH("Icons/icon_file_open_40x", Icon24x24));
-	StyleInstance->Set("Menu.Background", new BOX_BRUSH("Old/Menu_Background", FMargin(8.0f / 64.0f)));
-	StyleInstance->Set("BoldFont", DEFAULT_FONT("Bold", 9));
 
 	FSlateStyleRegistry::RegisterSlateStyle(*StyleInstance.Get());
+
+	SetParentStyleName(FEditorStyle::GetStyleSetName());
 }
 
 
diff --git a/Engine/Source/Developer/CollisionAnalyzer/Private/SCAQueryDetails.cpp b/Engine/Source/Developer/CollisionAnalyzer/Private/SCAQueryDetails.cpp
--- a/Engine/Source/Developer/CollisionAnalyzer/Private/SCAQueryDetails.cpp
+++ b/Engine/Source/Developer/CollisionAnalyzer/Private/SCAQueryDetails.cpp
@@ -161,7 +161,7 @@ void SCAQueryDetails::Construct(const FArguments& InArgs, TSharedPtr<SCollisionA
 		.AutoHeight()
 		[
 			SNew( SBorder )
-			.BorderImage(FCollisionAnalyzerStyle::Get()->GetBrush("ToolBar.Background"))
+			.BorderImage(FCollisionAnalyzerStyle::Get()->GetBrush("ToolPanel.GroupBorder"))
 			[
 				SNew(SHorizontalBox)
 				// Left is start/end locations
diff --git a/Engine/Source/Developer/CollisionAnalyzer/Private/SCollisionAnalyzer.cpp b/Engine/Source/Developer/CollisionAnalyzer/Private/SCollisionAnalyzer.cpp
--- a/Engine/Source/Developer/CollisionAnalyzer/Private/SCollisionAnalyzer.cpp
+++ b/Engine/Source/Developer/CollisionAnalyzer/Private/SCollisionAnalyzer.cpp
@@ -49,7 +49,7 @@ void SCollisionAnalyzer::Construct(const FArguments& InArgs, FCollisionAnalyzer*
 		.AutoHeight()
 		[
 			SNew( SBorder )
-			.BorderImage( FCollisionAnalyzerStyle::Get()->GetBrush( "ToolBar.Background" ) )
+			.BorderImage( FCollisionAnalyzerStyle::Get()->GetBrush( "ToolPanel.GroupBorder" ) )
 			[
 				SNew(SHorizontalBox)
 				// Record button
@@ -341,7 +341,7 @@ void SCollisionAnalyzer::Construct(const FArguments& InArgs, FCollisionAnalyzer*
 		.AutoHeight()
 		[
 			SNew( SBorder )
-			.BorderImage(FCollisionAnalyzerStyle::Get()->GetBrush("ToolBar.Background"))
+			.BorderImage(FCollisionAnalyzerStyle::Get()->GetBrush("ToolPanel.GroupBorder"))
 			[
 				SNew(STextBlock)
 				.Text(this, &SCollisionAnalyzer::GetStatusText)
diff --git a/Engine/Source/Editor/DeviceProfileEditor/Private/DetailsPanel/SDeviceProfileDetailsPanel.cpp b/Engine/Source/Editor/DeviceProfileEditor/Private/DetailsPanel/SDeviceProfileDetailsPanel.cpp
--- a/Engine/Source/Editor/DeviceProfileEditor/Private/DetailsPanel/SDeviceProfileDetailsPanel.cpp
+++ b/Engine/Source/Editor/DeviceProfileEditor/Private/DetailsPanel/SDeviceProfileDetailsPanel.cpp
@@ -108,7 +108,7 @@ void SDeviceProfileDetailsPanel::RefreshUI()
 		DetailsViewBox->AddSlot()
 		[
 			SNew( SBorder )
-			.BorderImage( FEditorStyle::GetBrush( "ToolBar.Background" ) )
+			.BorderImage( FEditorStyle::GetBrush( "ToolPanel.GroupBorder" ) )
 			[
 				SNew( SVerticalBox )
 				+ SVerticalBox::Slot()
@@ -168,7 +168,7 @@ void SDeviceProfileDetailsPanel::RefreshUI()
 		DetailsViewBox->AddSlot()
 			[
 				SNew( SBorder )
-				.BorderImage( FEditorStyle::GetBrush( "ToolBar.Background" ) )
+				.BorderImage( FEditorStyle::GetBrush( "ToolPanel.GroupBorder" ) )
 				[
 					SNew(SVerticalBox)
 					+SVerticalBox::Slot()
diff --git a/Engine/Source/Editor/DeviceProfileEditor/Private/SDeviceProfileCreateProfilePanel.cpp b/Engine/Source/Editor/DeviceProfileEditor/Private/SDeviceProfileCreateProfilePanel.cpp
--- a/Engine/Source/Editor/DeviceProfileEditor/Private/SDeviceProfileCreateProfilePanel.cpp
+++ b/Engine/Source/Editor/DeviceProfileEditor/Private/SDeviceProfileCreateProfilePanel.cpp
@@ -69,7 +69,7 @@ void SDeviceProfileCreateProfilePanel::Construct( const FArguments& InArgs, TWea
 		.FillHeight(1.0f)
 		[
 			SNew( SBorder )
-			.BorderImage( FEditorStyle::GetBrush( "ToolBar.Background" ) )
+			.BorderImage( FEditorStyle::GetBrush( "ToolPanel.GroupBorder" ) )
 			[
 				// Name entry
 				SNew( SVerticalBox )
diff --git a/Engine/Source/Editor/DeviceProfileEditor/Private/SDeviceProfileEditor.cpp b/Engine/Source/Editor/DeviceProfileEditor/Private/SDeviceProfileEditor.cpp
--- a/Engine/Source/Editor/DeviceProfileEditor/Private/SDeviceProfileEditor.cpp
+++ b/Engine/Source/Editor/DeviceProfileEditor/Private/SDeviceProfileEditor.cpp
@@ -381,7 +381,7 @@ TSharedRef<SDockTab> SDeviceProfileEditor::HandleTabManagerSpawnTab( const FSpaw
 	if( TabIdentifier == DeviceProfileEditorTabName )
 	{
 		TabWidget = SNew(SBorder)
-			.BorderImage(FEditorStyle::GetBrush("ToolBar.Background"))
+			.BorderImage(FEditorStyle::GetBrush("ToolPanel.GroupBorder"))
 			[
 				SNew(SOverlay)
 				+ SOverlay::Slot()
@@ -408,7 +408,7 @@ TSharedRef<SDockTab> SDeviceProfileEditor::HandleTabManagerSpawnTab( const FSpaw
 					+ SVerticalBox::Slot()
 					[
 						SNew(SBorder)
-						.BorderImage(FEditorStyle::GetBrush("ToolBar.Background"))
+						.BorderImage(FEditorStyle::GetBrush("ToolPanel.GroupBorder"))
 						[
 							SNew(SHorizontalBox)
 							+ SHorizontalBox::Slot()
diff --git a/Engine/Source/Editor/DeviceProfileEditor/Private/SDeviceProfileSelectionPanel.cpp b/Engine/Source/Editor/DeviceProfileEditor/Private/SDeviceProfileSelectionPanel.cpp
--- a/Engine/Source/Editor/DeviceProfileEditor/Private/SDeviceProfileSelectionPanel.cpp
+++ b/Engine/Source/Editor/DeviceProfileEditor/Private/SDeviceProfileSelectionPanel.cpp
@@ -256,7 +256,7 @@ void SDeviceProfileSelectionPanel::Construct( const FArguments& InArgs, TWeakObj
 			+ SScrollBox::Slot()
 			[
 				SNew( SBorder )
-				.BorderImage( FEditorStyle::GetBrush( "ToolBar.Background" ) )
+				.BorderImage( FEditorStyle::GetBrush( "ToolPanel.GroupBorder" ) )
 				[
 					SAssignNew( ListWidget, SVerticalBox )
 				]
diff --git a/Engine/Source/Editor/UndoHistory/Private/Widgets/SUndoHistoryDetails.cpp b/Engine/Source/Editor/UndoHistory/Private/Widgets/SUndoHistoryDetails.cpp
--- a/Engine/Source/Editor/UndoHistory/Private/Widgets/SUndoHistoryDetails.cpp
+++ b/Engine/Source/Editor/UndoHistory/Private/Widgets/SUndoHistoryDetails.cpp
@@ -240,7 +240,7 @@ void SUndoHistoryDetails::Construct(const FArguments& InArgs)
 	ChildSlot
 		[
 			SNew(SBorder)
-				.BorderImage(FEditorStyle::GetBrush("ToolBar.Background"))
+				.BorderImage(FEditorStyle::GetBrush("ToolPanel.GroupBorder"))
 				.Padding(FMargin(4.0f, 1.0f))
 				[
 					SNew(SVerticalBox)
@@ -252,7 +252,7 @@ void SUndoHistoryDetails::Construct(const FArguments& InArgs)
 						.AutoHeight()
 						[
 							SNew(SBorder)
-								.BorderImage(FEditorStyle::GetBrush("ToolBar.Background"))
+								.BorderImage(FEditorStyle::GetBrush("ToolPanel.GroupBorder"))
 								.Padding(FMargin(0.0f, 1.0f))
 								[
 									SNew(SGridPanel)
